Create base QAModeling Plugin

#rnx
#changelist validated

diff --git a/Engine/Plugins/QA-Modeling-Test-Assets/QA-Modeling-Test-Assets.uplugin b/Engine/Plugins/QA-Modeling-Test-Assets/QA-Modeling-Test-Assets.uplugin
new file mode 100644
--- /dev/null
+++ b/Engine/Plugins/QA-Modeling-Test-Assets/QA-Modeling-Test-Assets.uplugin
@@ -0,0 +1,17 @@
+{
+	"FileVersion": 3,
+	"Version": 1,
+	"VersionName": "1.0",
+	"FriendlyName": "QA-Modeling-Test-Assets",
+	"Description": "Contains QA Modeling Test Content so that Modeling Tests can be executed in a blank project.",
+	"Category": "Other",
+	"CreatedBy": "Alex Berryhill",
+	"CreatedByURL": "",
+	"DocsURL": "",
+	"MarketplaceURL": "",
+	"SupportURL": "",
+	"CanContainContent": true,
+	"IsBetaVersion": true,
+	"IsExperimentalVersion": false,
+	"Installed": false
+}
\ No newline at end of file
